[AI Toolsets]: Move StateTreeToolset under the Toolsets directory

#Synced-CL 52368392

diff --git a/Engine/Plugins/Experimental/StateTreeToolset/Source/StateTreeToolset/Private/StateTreeToolset.cpp b/Engine/Plugins/Experimental/StateTreeToolset/Source/StateTreeToolset/Private/StateTreeToolset.cpp
deleted file mode 100644
diff --git a/Engine/Plugins/Experimental/StateTreeToolset/Source/StateTreeToolset/Public/StateTreeToolset.h b/Engine/Plugins/Experimental/StateTreeToolset/Source/StateTreeToolset/Public/StateTreeToolset.h
deleted file mode 100644
diff --git a/Engine/Plugins/Experimental/StateTreeToolset/Source/StateTreeToolset/StateTreeToolset.Build.cs b/Engine/Plugins/Experimental/StateTreeToolset/Source/StateTreeToolset/StateTreeToolset.Build.cs
deleted file mode 100644
diff --git a/Engine/Plugins/Experimental/StateTreeToolset/StateTreeToolset.uplugin b/Engine/Plugins/Experimental/StateTreeToolset/StateTreeToolset.uplugin
deleted file mode 100644
diff --git a/Engine/Plugins/Experimental/Toolsets/StateTreeToolset/Source/StateTreeToolset/Private/StateTreeToolset.cpp b/Engine/Plugins/Experimental/Toolsets/StateTreeToolset/Source/StateTreeToolset/Private/StateTreeToolset.cpp
new file mode 100644
--- /dev/null
+++ b/Engine/Plugins/Experimental/Toolsets/StateTreeToolset/Source/StateTreeToolset/Private/StateTreeToolset.cpp
@@ -0,0 +1,16 @@
+// Copyright Epic Games, Inc. All Rights Reserved.
+
+#include "StateTreeToolset.h"
+
+#include "Modules/ModuleManager.h"
+
+/*virtual*/ void FStateTreeToolsetModule::StartupModule()
+{
+}
+
+/*virtual*/ void FStateTreeToolsetModule::ShutdownModule()
+{
+}
+
+IMPLEMENT_MODULE(FStateTreeToolsetModule, StateTreeToolset);
+DEFINE_LOG_CATEGORY(LogStateTreeToolset);
diff --git a/Engine/Plugins/Experimental/Toolsets/StateTreeToolset/Source/StateTreeToolset/Public/StateTreeToolset.h b/Engine/Plugins/Experimental/Toolsets/StateTreeToolset/Source/StateTreeToolset/Public/StateTreeToolset.h
new file mode 100644
--- /dev/null
+++ b/Engine/Plugins/Experimental/Toolsets/StateTreeToolset/Source/StateTreeToolset/Public/StateTreeToolset.h
@@ -0,0 +1,15 @@
+// Copyright Epic Games, Inc. All Rights Reserved.
+
+#pragma once
+
+#include "Modules/ModuleInterface.h"
+#include "Logging/LogMacros.h"
+
+DECLARE_LOG_CATEGORY_EXTERN(LogStateTreeToolset, Log, Log);
+
+class FStateTreeToolsetModule : public IModuleInterface
+{
+public:
+	virtual void StartupModule() override;
+	virtual void ShutdownModule() override;
+};
diff --git a/Engine/Plugins/Experimental/Toolsets/StateTreeToolset/Source/StateTreeToolset/StateTreeToolset.Build.cs b/Engine/Plugins/Experimental/Toolsets/StateTreeToolset/Source/StateTreeToolset/StateTreeToolset.Build.cs
new file mode 100644
--- /dev/null
+++ b/Engine/Plugins/Experimental/Toolsets/StateTreeToolset/Source/StateTreeToolset/StateTreeToolset.Build.cs
@@ -0,0 +1,12 @@
+// Copyright Epic Games, Inc. All Rights Reserved.
+
+using UnrealBuildTool;
+
+public class StateTreeToolset : ModuleRules
+{
+	public StateTreeToolset(ReadOnlyTargetRules Target) : base(Target)
+	{
+		PCHUsage = PCHUsageMode.UseExplicitOrSharedPCHs;
+		PublicDependencyModuleNames.Add("Core");
+	}
+}
diff --git a/Engine/Plugins/Experimental/Toolsets/StateTreeToolset/StateTreeToolset.uplugin b/Engine/Plugins/Experimental/Toolsets/StateTreeToolset/StateTreeToolset.uplugin
new file mode 100644
--- /dev/null
+++ b/Engine/Plugins/Experimental/Toolsets/StateTreeToolset/StateTreeToolset.uplugin
@@ -0,0 +1,38 @@
+{
+	"FileVersion": 3,
+	"Version": 1,
+	"VersionName": "1.0",
+	"FriendlyName": "StateTreeToolset",
+	"Description": "Toolset for StateTree Inspection",
+	"Category": "Other",
+	"CreatedBy": "Epic Games",
+	"CreatedByURL": "",
+	"DocsURL": "",
+	"MarketplaceURL": "",
+	"SupportURL": "",
+	"EditorOnly": true,
+	"CanContainContent": true,
+	"IsBetaVersion": false,
+	"IsExperimentalVersion": true,
+	"NoRedist": true,
+	"Installed": false,
+	"EnabledByDefault": false,
+	"Modules": [
+		{
+			"Name": "StateTreeToolset",
+			"Type": "Editor",
+			"LoadingPhase": "Default",
+			"TargetAllowList": ["Editor"]
+		}
+	],
+	"Plugins": [
+		{
+			"Name": "StateTree",
+			"Enabled": true
+		},
+		{
+			"Name": "ToolsetRegistry",
+			"Enabled": true
+		}
+	]
+}
